Bugfix for CL 50635208

#Synced-CL 50635208

diff --git a/Engine/Plugins/Animation/AudioDrivenAnimation/StreamingADA/Source/SpeechAnimationSolver/Private/SpeechAnimationSolverV3.cpp b/Engine/Plugins/Animation/AudioDrivenAnimation/StreamingADA/Source/SpeechAnimationSolver/Private/SpeechAnimationSolverV3.cpp
--- a/Engine/Plugins/Animation/AudioDrivenAnimation/StreamingADA/Source/SpeechAnimationSolver/Private/SpeechAnimationSolverV3.cpp
+++ b/Engine/Plugins/Animation/AudioDrivenAnimation/StreamingADA/Source/SpeechAnimationSolver/Private/SpeechAnimationSolverV3.cpp
@@ -310,7 +310,7 @@ bool FSpeechAnimationSolverV3::ReformatAudio(const FSpeechAnimationAudioFrame& I
 
 	if (AudioFrameReformatted.SampleRate != SampleRate)
 	{
-		if (AudioFrameReformatted.SampleRate < KINDA_SMALL_NUMBER)
+		if (float(AudioFrameReformatted.SampleRate) < KINDA_SMALL_NUMBER)
 		{
 			UE_LOG(LogSpeechAnimationSolverV3, Warning, TEXT("Sample rate too low"));
 			return false;
